[InlineMaterialInstance] Split feature from VirtualProductionUtilities

diff --git a/Engine/Plugins/Experimental/InlineMaterialInstance/InlineMaterialInstance.uplugin b/Engine/Plugins/Experimental/InlineMaterialInstance/InlineMaterialInstance.uplugin
new file mode 100644
--- /dev/null
+++ b/Engine/Plugins/Experimental/InlineMaterialInstance/InlineMaterialInstance.uplugin
@@ -0,0 +1,23 @@
+﻿{
+	"FileVersion": 3,
+	"Version": 1,
+	"VersionName": "1.0",
+	"FriendlyName": "Inline Material Instance",
+	"Description": "Allows creation and modification of materials within the property editor",
+	"Category": "Virtual Production",
+	"CreatedBy": "Epic Games, Inc.",
+	"CreatedByURL": "https://epicgames.com",
+	"DocsURL": "",
+	"MarketplaceURL": "",
+	"SupportURL": "",
+	"CanContainContent": true,
+	"IsBetaVersion": true,
+	"Installed": false,
+	"Modules": [
+		{
+			"Name": "InlineMaterialInstanceEditor",
+			"Type": "Editor",
+			"LoadingPhase": "PostDefault"
+		}
+	]
+}
diff --git a/Engine/Plugins/Experimental/InlineMaterialInstance/Source/InlineMaterialInstanceEditor/InlineMaterialInstanceEditor.Build.cs b/Engine/Plugins/Experimental/InlineMaterialInstance/Source/InlineMaterialInstanceEditor/InlineMaterialInstanceEditor.Build.cs
new file mode 100644
--- /dev/null
+++ b/Engine/Plugins/Experimental/InlineMaterialInstance/Source/InlineMaterialInstanceEditor/InlineMaterialInstanceEditor.Build.cs
@@ -0,0 +1,27 @@
+// Copyright Epic Games, Inc. All Rights Reserved.
+
+using UnrealBuildTool;
+
+public class InlineMaterialInstanceEditor : ModuleRules
+{
+	public InlineMaterialInstanceEditor(ReadOnlyTargetRules Target) : base(Target)
+	{
+		PCHUsage = ModuleRules.PCHUsageMode.UseExplicitOrSharedPCHs;
+
+		PrivateDependencyModuleNames.AddRange(
+			new string[]
+			{
+				"Core",
+				"CoreUObject",
+				"Engine",
+				"InputCore",
+				"Landscape",
+				"PropertyEditor",
+				"MaterialEditor",
+				"Slate",
+				"SlateCore",
+				"UnrealEd",
+            }
+		);
+	}
+}
diff --git a/Engine/Plugins/Experimental/InlineMaterialInstance/Source/InlineMaterialInstanceEditor/Private/InlineMaterialInstanceEditorModule.cpp b/Engine/Plugins/Experimental/InlineMaterialInstance/Source/InlineMaterialInstanceEditor/Private/InlineMaterialInstanceEditorModule.cpp
new file mode 100644
--- /dev/null
+++ b/Engine/Plugins/Experimental/InlineMaterialInstance/Source/InlineMaterialInstanceEditor/Private/InlineMaterialInstanceEditorModule.cpp
@@ -0,0 +1,36 @@
+// Copyright Epic Games, Inc. All Rights Reserved.
+
+#include "HAL/LowLevelMemTracker.h"
+#include "MaterialList.h"
+#include "Modules/ModuleManager.h"
+#include "Widgets/SMaterialDynamicWidgets.h"
+
+LLM_DEFINE_TAG(InlineMaterialInstanceEditor);
+
+#define LOCTEXT_NAMESPACE "FInlineMaterialInstanceEditorModule"
+
+class FInlineMaterialInstanceEditorModule : public IModuleInterface
+{
+public:
+	void StartupModule()
+	{
+		LLM_SCOPE_BYTAG(InlineMaterialInstanceEditor);
+
+		// Add bottom extender for material item
+		FMaterialList::OnAddMaterialItemViewExtraBottomWidget.AddLambda([](const TSharedRef<FMaterialItemView>& InMaterialItemView, UActorComponent* InCurrentComponent, IDetailLayoutBuilder& InDetailBuilder, TArray<TSharedPtr<SWidget>>& OutExtensions)
+		{
+			LLM_SCOPE_BYTAG(InlineMaterialInstanceEditor);
+			OutExtensions.Add(SNew(SMaterialDynamicView, InMaterialItemView, InCurrentComponent));
+		});
+	}
+
+	void ShutdownModule()
+	{
+		LLM_SCOPE_BYTAG(InlineMaterialInstanceEditor);
+		FMaterialList::OnAddMaterialItemViewExtraBottomWidget.RemoveAll(this);
+	}
+};
+
+#undef LOCTEXT_NAMESPACE
+	
+IMPLEMENT_MODULE(FInlineMaterialInstanceEditorModule, InlineMaterialInstanceEditor)
diff --git a/Engine/Plugins/Experimental/InlineMaterialInstance/Source/InlineMaterialInstanceEditor/Private/Widgets/SMaterialDynamicParametersPanelWidget.cpp b/Engine/Plugins/Experimental/InlineMaterialInstance/Source/InlineMaterialInstanceEditor/Private/Widgets/SMaterialDynamicParametersPanelWidget.cpp
new file mode 100644
--- /dev/null
+++ b/Engine/Plugins/Experimental/InlineMaterialInstance/Source/InlineMaterialInstanceEditor/Private/Widgets/SMaterialDynamicParametersPanelWidget.cpp
@@ -0,0 +1,763 @@
+// Copyright Epic Games, Inc. All Rights Reserved.
+
+#include "SMaterialDynamicParametersPanelWidget.h"
+
+#include "DetailColumnSizeData.h"
+#include "DetailWidgetRow.h"
+#include "Styling/AppStyle.h"
+#include "IDetailPropertyRow.h"
+#include "IDetailTreeNode.h"
+#include "IPropertyRowGenerator.h"
+#include "PropertyCustomizationHelpers.h"
+#include "PropertyEditorModule.h"
+#include "PropertyHandle.h"
+#include "Framework/MultiBox/MultiBoxBuilder.h"
+#include "Framework/MultiBox/MultiBoxDefs.h"
+#include "Materials/MaterialInstance.h"
+#include "Modules/ModuleManager.h"
+#include "UObject/Object.h"
+#include "Widgets/Layout/SWidgetSwitcher.h"
+#include "Widgets/Text/SInlineEditableTextBlock.h"
+#include "Widgets/Views/SExpanderArrow.h"
+#include "Widgets/Views/STreeView.h"
+
+#define LOCTEXT_NAMESPACE "MaterialDynamicParametersPanelWidget"
+
+enum EStackDataType
+{
+	Stack,
+	Asset,
+	Group,
+	Property,
+	PropertyChild,
+};
+
+struct FPropertySortedParamData
+{
+	EStackDataType StackDataType;
+
+	FMaterialParameterInfo ParameterInfo;
+
+	TSharedPtr<IDetailTreeNode> ParameterNode;
+
+	TSharedPtr<IPropertyHandle> ParameterHandle;
+
+	TArray<TSharedPtr<FPropertySortedParamData>> Children;
+
+	FString NodeKey;
+
+	FName GroupName;
+};
+
+struct FPropertyUnsortedParamData
+{
+	TSharedPtr<IDetailTreeNode> ParameterNode;
+
+	TSharedPtr<IPropertyHandle> ParameterHandle;
+
+	FMaterialParameterInfo ParameterInfo;
+
+	FName GroupName;
+};
+
+namespace Utilities
+{
+	template<typename ParameterValueType>
+	void AddParameters(TSharedPtr<IPropertyHandle> InPropertyHandle, const FName InParameterName, TArray<TSharedPtr<IPropertyHandle>>& OutDeferredSearches, TArray<FPropertyUnsortedParamData>& OutUnsortedParameters)
+	{
+		if (!InPropertyHandle.IsValid())
+		{
+			return;
+		}
+		
+		TArray<void*> GroupPtrs;
+		InPropertyHandle->AccessRawData(GroupPtrs);
+		const auto GroupIt = GroupPtrs.CreateConstIterator();
+		TArray<ParameterValueType>* ParameterGroupArrayPtr = reinterpret_cast<TArray<ParameterValueType>*>(*GroupIt);
+		if (!ParameterGroupArrayPtr)
+		{
+			return;
+		}
+
+		const TSharedPtr<IPropertyHandleArray> ParameterValuesArray = InPropertyHandle->AsArray();
+
+		uint32 NumElements = 0;
+		ParameterValuesArray->GetNumElements(NumElements);
+		for (uint32 Index = 0; Index < NumElements; ++Index)
+		{
+			const ParameterValueType& ParameterValue = (*ParameterGroupArrayPtr)[Index];
+			if (!ParameterValue.ExpressionGUID.IsValid())
+			{
+				continue;
+			}
+
+			if (const TSharedPtr<IPropertyHandle> ElementPropertyHandle = ParameterValuesArray->GetElement(Index))
+			{
+				FPropertyUnsortedParamData NonLayerProperty;
+				TSharedPtr<IPropertyHandle> ParameterValuePropertyHandle = ElementPropertyHandle->GetChildHandle(InParameterName);
+
+				NonLayerProperty.ParameterInfo = ParameterValue.ParameterInfo;
+				NonLayerProperty.GroupName = TEXT("Group");
+
+				OutDeferredSearches.Add(ParameterValuePropertyHandle);
+				OutUnsortedParameters.Add(NonLayerProperty);
+			}
+		}
+	}
+};
+
+// ********* SMaterialParametersOverviewTree *******
+class SMaterialDynamicParametersOverviewTree : public STreeView<TSharedPtr<FPropertySortedParamData>>
+{
+	friend class SMaterialDynamicParametersOverviewTreeItem;
+
+public:
+
+	SLATE_BEGIN_ARGS(SMaterialDynamicParametersOverviewTree)
+	{}
+
+	SLATE_ARGUMENT(UMaterialInstance*, InMaterialInstance)
+	SLATE_ARGUMENT(TSharedPtr<SMaterialDynamicParametersPanelWidget>, InOwner)
+	SLATE_ARGUMENT(TSharedPtr<SScrollBar>, InScrollbar)
+	SLATE_END_ARGS()
+
+	/** Constructs this widget with InArgs */
+	void Construct(const FArguments& InArgs);
+	TSharedRef< ITableRow > OnGenerateRowMaterialLayersFunctionsTreeView(TSharedPtr<FPropertySortedParamData> Item, const TSharedRef< STableViewBase >& OwnerTable);
+	void OnGetChildrenMaterialLayersFunctionsTreeView(TSharedPtr<FPropertySortedParamData> InParent, TArray< TSharedPtr<FPropertySortedParamData> >& OutChildren);
+	void OnExpansionChanged(TSharedPtr<FPropertySortedParamData> Item, bool bIsExpanded);
+	void SetParentsExpansionState();
+
+	/** Builds the custom parameter groups category */
+	void CreateGroupsWidget();
+
+	TWeakPtr<SMaterialDynamicParametersPanelWidget> GetOwner() { return Owner; }
+	bool HasAnyParameters() const { return bHasAnyParameters; }
+
+	FDetailColumnSizeData& GetColumnSizeData() { return ColumnSizeData; }
+
+protected:
+
+	void ShowSubParameters();
+
+public:
+	/** Object that stores all of the possible parameters we can edit */
+	TWeakObjectPtr<UMaterialInstance> MaterialInstance;
+
+private:
+	TArray<TSharedPtr<FPropertySortedParamData>> SortedParameters;
+
+	TArray<FPropertyUnsortedParamData> UnsortedParameters;
+
+	TWeakPtr<SMaterialDynamicParametersPanelWidget> Owner;
+
+	bool bHasAnyParameters;
+	
+	FDetailColumnSizeData ColumnSizeData;
+};
+
+// ********* SMaterialParametersOverviewTreeItem *******
+class SMaterialDynamicParametersOverviewTreeItem : public STableRow< TSharedPtr<FPropertySortedParamData> >
+{
+public:
+
+	SLATE_BEGIN_ARGS(SMaterialDynamicParametersOverviewTreeItem)
+		: _StackParameterData(nullptr),
+		_MaterialInstance(nullptr)
+	{}
+
+	/** The item content. */
+	SLATE_ARGUMENT(TSharedPtr<FPropertySortedParamData>, StackParameterData)
+	SLATE_ARGUMENT(UMaterialInstance*, MaterialInstance)
+	SLATE_ARGUMENT(TSharedPtr<SMaterialDynamicParametersOverviewTree>, InTree)
+	SLATE_END_ARGS()
+
+	void RefreshOnRowChange(const FAssetData& AssetData, TSharedPtr<SMaterialDynamicParametersOverviewTree> InTree);
+
+	/**
+	* Construct the widget
+	*
+	* @param InArgs   A declaration from which to construct the widget
+	*/
+	void Construct(const FArguments& InArgs, const TSharedRef<STableViewBase>& InOwnerTableView);
+
+private:
+	const FSlateBrush* GetBorderImage() const;
+
+	FSlateColor GetOuterBackgroundColor(TSharedPtr<FPropertySortedParamData> InParamData) const;
+
+	TSharedRef<SWidget> GetRowExtensionButtons(TSharedPtr<IPropertyHandle> InPropertyHandle);
+private:
+
+	/** The node info to build the tree view row from. */
+	TSharedPtr<FPropertySortedParamData> StackParameterData;
+
+	/** The tree that contains this item */
+	TWeakPtr<SMaterialDynamicParametersOverviewTree> Tree;
+
+	/** The set of material parameters this is associated with */
+	TWeakObjectPtr<UMaterialInstance> MaterialInstance;
+
+	/** Pointer to copied ValuePtr. Needed for workaround for now */
+	TUniquePtr<uint8> NewDefaultValuePtr;
+
+	/** Pointer to Widget for Workaround. */
+	TSharedPtr<SWidget> ResetArrow;
+};
+
+
+
+const FSlateBrush* SMaterialDynamicParametersOverviewTreeItem::GetBorderImage() const
+{
+	return FAppStyle::Get().GetBrush("DetailsView.CategoryMiddle");
+}
+
+FSlateColor SMaterialDynamicParametersOverviewTreeItem::GetOuterBackgroundColor(TSharedPtr<FPropertySortedParamData> InParamData) const
+{
+	if (IsHovered() || InParamData->StackDataType == EStackDataType::Group)
+	{
+		return FAppStyle::Get().GetSlateColor("Colors.Header");
+	}
+
+	return FAppStyle::Get().GetSlateColor("Colors.Panel");
+}
+
+TSharedRef<SWidget> SMaterialDynamicParametersOverviewTreeItem::GetRowExtensionButtons(TSharedPtr<IPropertyHandle> InPropertyHandle)
+{
+	if (!InPropertyHandle.IsValid())
+	{
+		return SNullWidget::NullWidget;
+	}
+
+	const auto ResetDelegate = FExecuteAction::CreateLambda([this, InPropertyHandle]()
+	{
+		if (InPropertyHandle.IsValid())
+		{
+			InPropertyHandle->ResetToDefault();
+			ResetArrow->SetVisibility(EVisibility::Hidden);
+		}
+	});
+
+	/** Needed as Dynamic Material Instances don't have a default to use normally for DiffersFromDefault */
+	const auto DiffersFromDefaultDelegate = FIsActionButtonVisible::CreateLambda([this, InPropertyHandle]()
+	{
+		if (!NewDefaultValuePtr)
+		{
+			return true;
+		}
+
+		FProperty* Property = InPropertyHandle->GetProperty();
+		
+		void* ValuePtr;
+		InPropertyHandle->GetValueData(ValuePtr);
+		
+		if (Property->Identical(NewDefaultValuePtr.Get(), ValuePtr))
+		{
+			return false;
+		}
+		return true;
+	});
+	
+	FPropertyEditorModule& PropertyEditorModule = FModuleManager::Get().GetModuleChecked<FPropertyEditorModule>("PropertyEditor");
+
+	TArray<FPropertyRowExtensionButton> ExtensionButtons;
+	FPropertyRowExtensionButton& ResetToDefault = ExtensionButtons.AddDefaulted_GetRef();
+	ResetToDefault.Label = NSLOCTEXT("PropertyEditor", "ResetToDefault", "Reset to Default");
+	ResetToDefault.UIAction = FUIAction(ResetDelegate, FCanExecuteAction(), FIsActionChecked(), DiffersFromDefaultDelegate);
+
+	ResetToDefault.Icon = FSlateIcon(FAppStyle::Get().GetStyleSetName(), "PropertyWindow.DiffersFromDefault");
+	ResetToDefault.ToolTip = NSLOCTEXT("PropertyEditor", "ResetToDefaultToolTip", "Reset this property to its default value.");
+
+	// Add custom property extensions
+	FOnGenerateGlobalRowExtensionArgs RowExtensionArgs;
+	RowExtensionArgs.PropertyHandle = InPropertyHandle;
+	PropertyEditorModule.GetGlobalRowExtensionDelegate().Broadcast(RowExtensionArgs, ExtensionButtons);
+	
+	// Build extension toolbar 
+	FSlimHorizontalToolBarBuilder ToolbarBuilder(TSharedPtr<FUICommandList>(), FMultiBoxCustomization::None);
+	ToolbarBuilder.SetLabelVisibility(EVisibility::Collapsed);
+	ToolbarBuilder.SetStyle(&FAppStyle::Get(), "DetailsView.ExtensionToolBar");
+	for (const FPropertyRowExtensionButton& Extension : ExtensionButtons)
+	{
+		ToolbarBuilder.AddToolBarButton(Extension.UIAction, NAME_None, Extension.Label, Extension.ToolTip, Extension.Icon);
+	}
+
+	/** Set up the Default Value for the Workaround. */
+	FProperty* Property = InPropertyHandle->GetProperty();
+	NewDefaultValuePtr = MakeUnique<uint8>(Property->GetSize());
+	void* ValuePtr;
+	InPropertyHandle->GetValueData(ValuePtr);
+	Property->CopyCompleteValue(NewDefaultValuePtr.Get(), ValuePtr);
+	Property->ClearValue(NewDefaultValuePtr.Get());
+	
+	ResetArrow = ToolbarBuilder.MakeWidget();
+	
+	return ResetArrow.ToSharedRef();
+}
+
+void SMaterialDynamicParametersOverviewTreeItem::RefreshOnRowChange(const FAssetData& AssetData, TSharedPtr<SMaterialDynamicParametersOverviewTree> InTree)
+{
+	if (InTree.IsValid())
+	{
+		InTree->CreateGroupsWidget();
+	}
+}
+
+void SMaterialDynamicParametersOverviewTreeItem::Construct(const FArguments& InArgs, const TSharedRef<STableViewBase>& InOwnerTableView)
+{
+	StackParameterData = InArgs._StackParameterData;
+	MaterialInstance = InArgs._MaterialInstance;
+	Tree = InArgs._InTree;
+
+	TSharedRef<SWidget> LeftSideWidget = SNullWidget::NullWidget;
+	TSharedRef<SWidget> RightSideWidget = SNullWidget::NullWidget;
+	FText NameOverride;
+	TSharedRef<SVerticalBox> WrapperWidget = SNew(SVerticalBox);
+
+	auto GenerateLeftAndRightWidgets = [&](IDetailTreeNode& InNode)
+	{
+		const FNodeWidgets NodeWidgets = InNode.CreateNodeWidgets();
+		LeftSideWidget = NodeWidgets.NameWidget.ToSharedRef();
+
+		RightSideWidget = SNew(SHorizontalBox)
+				+SHorizontalBox::Slot()
+				.FillWidth(1.f)
+				[
+					NodeWidgets.ValueWidget.ToSharedRef()
+				]
+				+ SHorizontalBox::Slot()
+				.AutoWidth()
+				[
+					GetRowExtensionButtons(InNode.CreatePropertyHandle())
+				];
+	};
+
+// GROUP --------------------------------------------------
+	if (StackParameterData->StackDataType == EStackDataType::Group)
+	{
+		NameOverride = FText::FromName(StackParameterData->GroupName);
+		LeftSideWidget = SNew(STextBlock)
+			.TransformPolicy(ETextTransformPolicy::ToUpper)
+			.Text(NameOverride)
+			.Font(FAppStyle::Get().GetFontStyle("PropertyWindow.BoldFont"))
+			.TextStyle(FAppStyle::Get(), "DetailsView.CategoryTextStyle");
+	}
+// END GROUP
+
+// PROPERTY ----------------------------------------------
+	bool bIsPaddedProperty = false;
+	if (StackParameterData->StackDataType == EStackDataType::Property)
+	{
+		NameOverride = FText::FromName(StackParameterData->ParameterInfo.Name);
+
+		IDetailTreeNode& Node = *StackParameterData->ParameterNode;
+		TSharedPtr<IDetailPropertyRow> GeneratedRow = StaticCastSharedPtr<IDetailPropertyRow>(Node.GetRow());
+		IDetailPropertyRow& Row = *GeneratedRow.Get();
+		Row.DisplayName(NameOverride);
+
+
+		if (TSharedPtr<IPropertyHandle> PropertyHandle = StackParameterData->ParameterNode->CreatePropertyHandle())
+		{
+			PropertyHandle->MarkResetToDefaultCustomized(true);
+		}
+
+		FDetailWidgetDecl* CustomNameWidget = Row.CustomNameWidget();
+		if (CustomNameWidget)
+		{
+			(*CustomNameWidget)
+			[
+				SNew(SHorizontalBox)
+				+SHorizontalBox::Slot()
+				.VAlign(VAlign_Center)
+				[
+					SNew(STextBlock)
+					.Text(NameOverride)
+					.Font(FAppStyle::GetFontStyle(TEXT("PropertyWindow.NormalFont")))
+				]
+			];
+		}
+
+		bIsPaddedProperty = true;
+
+		GenerateLeftAndRightWidgets(Node);
+	}
+// END PROPERTY
+
+// PROPERTY CHILD ----------------------------------------
+	if (StackParameterData->StackDataType == EStackDataType::PropertyChild)
+	{
+		GenerateLeftAndRightWidgets(*StackParameterData->ParameterNode);
+	}
+// END PROPERTY CHILD
+
+// FINAL WRAPPER
+
+	{
+		FDetailColumnSizeData& ColumnSizeData = InArgs._InTree->GetColumnSizeData();
+
+		float ValuePadding = bIsPaddedProperty ? 20.0f : 0.0f;
+		WrapperWidget->AddSlot()
+			.AutoHeight()
+			[
+				SNew(SBorder)
+				.BorderImage(FAppStyle::Get().GetBrush("DetailsView.GridLine"))
+				.Padding(FMargin(0, 0, 0, 1))
+				[
+					SNew(SBorder)
+					.Padding(3.0f)
+					.BorderImage(this, &SMaterialDynamicParametersOverviewTreeItem::GetBorderImage)
+					.BorderBackgroundColor(this, &SMaterialDynamicParametersOverviewTreeItem::GetOuterBackgroundColor, StackParameterData)
+					[
+						SNew(SSplitter)
+						.Style(FAppStyle::Get(), "DetailsView.Splitter")
+						.PhysicalSplitterHandleSize(1.0f)
+						.HitDetectionSplitterHandleSize(5.0f)
+						+ SSplitter::Slot()
+						.Value(ColumnSizeData.GetNameColumnWidth())
+						.OnSlotResized(ColumnSizeData.GetOnNameColumnResized())
+						.Value(0.25f)
+						[
+							SNew(SHorizontalBox)
+							+ SHorizontalBox::Slot()
+							.AutoWidth()
+							.VAlign(VAlign_Center)
+							.Padding(FMargin(3.0f))
+							[
+								SNew(SExpanderArrow, SharedThis(this))
+							]
+							+ SHorizontalBox::Slot()
+							.Padding(FMargin(2.0f))
+							.VAlign(VAlign_Center)
+							[
+								LeftSideWidget
+							]
+						]
+						+ SSplitter::Slot()
+						.Value(ColumnSizeData.GetValueColumnWidth())
+						.OnSlotResized(ColumnSizeData.GetOnValueColumnResized()) 
+						[
+							SNew(SHorizontalBox)
+							+ SHorizontalBox::Slot()
+							.Padding(FMargin(5.0f, 2.0f, 5.0, 2.0f))
+							.HAlign(HAlign_Left)
+							.VAlign(VAlign_Center)
+							[
+								RightSideWidget
+							]
+						]
+					]
+				]
+			];
+	}
+
+	ChildSlot
+		[
+			WrapperWidget
+		];
+
+
+	STableRow< TSharedPtr<FPropertySortedParamData> >::ConstructInternal(
+		STableRow::FArguments()
+		.Style(FAppStyle::Get(), "DetailsView.TreeView.TableRow")
+		.ShowSelection(false),
+		InOwnerTableView
+	);
+}
+
+void SMaterialDynamicParametersOverviewTree::Construct(const FArguments& InArgs)
+{
+	bHasAnyParameters = false;
+	MaterialInstance = InArgs._InMaterialInstance;
+	Owner = InArgs._InOwner;
+	CreateGroupsWidget();
+
+	STreeView<TSharedPtr<FPropertySortedParamData>>::Construct(
+		STreeView::FArguments()
+		.TreeItemsSource(&SortedParameters)
+		.SelectionMode(ESelectionMode::None)
+		.OnGenerateRow(this, &SMaterialDynamicParametersOverviewTree::OnGenerateRowMaterialLayersFunctionsTreeView)
+		.OnGetChildren(this, &SMaterialDynamicParametersOverviewTree::OnGetChildrenMaterialLayersFunctionsTreeView)
+		.OnExpansionChanged(this, &SMaterialDynamicParametersOverviewTree::OnExpansionChanged)
+		.ExternalScrollbar(InArgs._InScrollbar)
+	);
+}
+
+TSharedRef< ITableRow > SMaterialDynamicParametersOverviewTree::OnGenerateRowMaterialLayersFunctionsTreeView(TSharedPtr<FPropertySortedParamData> Item, const TSharedRef< STableViewBase >& OwnerTable)
+{
+	TSharedRef< SMaterialDynamicParametersOverviewTreeItem > ReturnRow = SNew(SMaterialDynamicParametersOverviewTreeItem, OwnerTable)
+		.StackParameterData(Item)
+		.MaterialInstance(MaterialInstance.Get())
+		.InTree(SharedThis(this));
+	return ReturnRow;
+}
+
+void SMaterialDynamicParametersOverviewTree::OnGetChildrenMaterialLayersFunctionsTreeView(TSharedPtr<FPropertySortedParamData> InParent, TArray< TSharedPtr<FPropertySortedParamData> >& OutChildren)
+{
+	OutChildren = InParent->Children;
+}
+
+void SMaterialDynamicParametersOverviewTree::OnExpansionChanged(TSharedPtr<FPropertySortedParamData> Item, bool bIsExpanded)
+{
+	bool* ExpansionValue = MaterialInstance->ParameterOverviewExpansion.Find(Item->NodeKey);
+	if (ExpansionValue == nullptr)
+	{
+		MaterialInstance->ParameterOverviewExpansion.Add(Item->NodeKey, bIsExpanded);
+	}
+	else if (*ExpansionValue != bIsExpanded)
+	{
+		MaterialInstance->ParameterOverviewExpansion.Emplace(Item->NodeKey, bIsExpanded);
+	}
+	// Expand any children that are also expanded
+	for (TSharedPtr<FPropertySortedParamData> Child : Item->Children)
+	{
+		bool* ChildExpansionValue = MaterialInstance->ParameterOverviewExpansion.Find(Child->NodeKey);
+		if (ChildExpansionValue != nullptr && *ChildExpansionValue == true)
+		{
+			SetItemExpansion(Child, true);
+		}
+	}
+}
+
+void SMaterialDynamicParametersOverviewTree::SetParentsExpansionState()
+{
+	for (const auto& Pair : SortedParameters)
+	{
+		if (Pair->Children.Num())
+		{
+			bool* bIsExpanded = MaterialInstance->ParameterOverviewExpansion.Find(Pair->NodeKey);
+			if (bIsExpanded)
+			{
+				SetItemExpansion(Pair, *bIsExpanded);
+			}
+			else
+			{
+				SetItemExpansion(Pair, true);
+			}
+		}
+	}
+}
+
+void SMaterialDynamicParametersOverviewTree::CreateGroupsWidget()
+{
+	check(MaterialInstance.IsValid());
+	UnsortedParameters.Reset();
+	SortedParameters.Reset();
+
+	const TArray<TSharedRef<IDetailTreeNode>> TestData = GetOwner().Pin()->Generator->GetRootTreeNodes();
+	if (TestData.Num() == 0)
+	{
+		return;
+	}
+	TSharedPtr<IDetailTreeNode> Category = TestData[0];
+	TSharedPtr<IDetailTreeNode> ParameterGroups;
+	TArray<TSharedRef<IDetailTreeNode>> Children;
+	Category->GetChildren(Children);
+
+	TArray<TSharedPtr<IPropertyHandle>> DeferredSearches;
+
+	for (int32 ChildIdx = 0; ChildIdx < Children.Num(); ChildIdx++)
+	{
+		TSharedPtr<IPropertyHandle> PropertyHandle = Children[ChildIdx]->CreatePropertyHandle();
+
+		bHasAnyParameters = true;
+
+		if (!PropertyHandle.IsValid() || !PropertyHandle->GetProperty())
+		{
+			continue;
+		}
+
+		if (PropertyHandle->GetProperty()->GetFName() == GET_MEMBER_NAME_CHECKED(UMaterialInstance, ScalarParameterValues))
+		{
+			Utilities::AddParameters<FScalarParameterValue>(PropertyHandle,
+			                                                 GET_MEMBER_NAME_CHECKED(
+				                                                 FScalarParameterValue, ParameterValue),
+			                                                 DeferredSearches, UnsortedParameters);
+		}
+		else if (PropertyHandle->GetProperty()->GetFName() == GET_MEMBER_NAME_CHECKED(UMaterialInstance, VectorParameterValues))
+		{
+			Utilities::AddParameters<FVectorParameterValue>(PropertyHandle,
+															GET_MEMBER_NAME_CHECKED(
+																FVectorParameterValue, ParameterValue),
+															DeferredSearches, UnsortedParameters);
+		}
+		else if (PropertyHandle->GetProperty()->GetFName() == GET_MEMBER_NAME_CHECKED(UMaterialInstance, TextureParameterValues))
+		{
+			Utilities::AddParameters<FTextureParameterValue>(PropertyHandle,
+															GET_MEMBER_NAME_CHECKED(
+																FTextureParameterValue, ParameterValue),
+															DeferredSearches, UnsortedParameters);
+		}
+		else if (PropertyHandle->GetProperty()->GetFName() == GET_MEMBER_NAME_CHECKED(UMaterialInstance, FontParameterValues))
+		{
+			Utilities::AddParameters<FFontParameterValue>(PropertyHandle,
+															GET_MEMBER_NAME_CHECKED(
+																FFontParameterValue, FontValue),
+															DeferredSearches, UnsortedParameters);
+		}
+	}
+
+	checkf(UnsortedParameters.Num() == DeferredSearches.Num(), TEXT("Internal inconsistency: number of node searches does not match the number of properties"));
+	TArray<TSharedPtr<IDetailTreeNode>> DeferredResults = GetOwner().Pin()->Generator->FindTreeNodes(DeferredSearches);
+	checkf(UnsortedParameters.Num() == DeferredResults.Num(), TEXT("Internal inconsistency: number of node search results does not match the number of properties"));
+
+	for (int Idx = 0, NumUnsorted = UnsortedParameters.Num(); Idx < NumUnsorted; ++Idx)
+	{
+		FPropertyUnsortedParamData& NonLayerProperty = UnsortedParameters[Idx];
+		NonLayerProperty.ParameterNode = DeferredResults[Idx];
+		NonLayerProperty.ParameterHandle = NonLayerProperty.ParameterNode->CreatePropertyHandle();
+	}
+
+	Children.Empty();
+
+	ShowSubParameters();
+	RequestTreeRefresh();
+	SetParentsExpansionState();
+}
+
+void SMaterialDynamicParametersOverviewTree::ShowSubParameters()
+{
+	for (FPropertyUnsortedParamData Property : UnsortedParameters)
+	{
+		TSharedPtr<FPropertySortedParamData> GroupProperty(new FPropertySortedParamData());
+		GroupProperty->ParameterInfo = Property.ParameterInfo;
+		GroupProperty->StackDataType = EStackDataType::Group;
+		GroupProperty->GroupName = Property.GroupName;
+		GroupProperty->NodeKey = FString::FromInt(GroupProperty->ParameterInfo.Index) + FString::FromInt(GroupProperty->ParameterInfo.Association) + Property.GroupName.ToString();
+
+		bool bAddNewGroup = true;
+		for (TSharedPtr<struct FPropertySortedParamData> GroupChild : SortedParameters)
+		{
+			if (GroupChild->NodeKey == GroupProperty->NodeKey)
+			{
+				bAddNewGroup = false;
+			}
+		}
+		if (bAddNewGroup)
+		{
+			SortedParameters.Add(GroupProperty);
+		}
+
+		TSharedPtr<FPropertySortedParamData> ChildProperty(new FPropertySortedParamData());
+		ChildProperty->ParameterInfo = Property.ParameterInfo;
+		ChildProperty->StackDataType = EStackDataType::Property;
+		ChildProperty->ParameterNode = Property.ParameterNode;
+		ChildProperty->NodeKey = FString::FromInt(ChildProperty->ParameterInfo.Index) + FString::FromInt(ChildProperty->ParameterInfo.Association) + Property.GroupName.ToString() + Property.ParameterInfo.Name.ToString();
+
+		TArray<TSharedRef<IDetailTreeNode>> ParamChildren;
+		Property.ParameterNode->GetChildren(ParamChildren);
+		for (int32 ParamChildIdx = 0; ParamChildIdx < ParamChildren.Num(); ParamChildIdx++)
+		{
+			TSharedPtr<FPropertySortedParamData> ParamChildProperty(new FPropertySortedParamData());
+			ParamChildProperty->StackDataType = EStackDataType::PropertyChild;
+			ParamChildProperty->ParameterNode = ParamChildren[ParamChildIdx];
+			ParamChildProperty->ParameterHandle = ParamChildProperty->ParameterNode->CreatePropertyHandle();
+			ParamChildProperty->ParameterInfo = Property.ParameterInfo;
+			ChildProperty->Children.Add(ParamChildProperty);
+		}
+
+		for (TSharedPtr<struct FPropertySortedParamData> GroupChild : SortedParameters)
+		{
+			if (GroupChild->GroupName == Property.GroupName
+				&& GroupChild->ParameterInfo.Association == ChildProperty->ParameterInfo.Association
+				&&  GroupChild->ParameterInfo.Index == ChildProperty->ParameterInfo.Index)
+			{
+				GroupChild->Children.Add(ChildProperty);
+			}
+		}
+	}
+}
+
+void SMaterialDynamicParametersPanelWidget::Refresh()
+{
+	TSharedPtr<SHorizontalBox> HeaderBox;
+	NestedTree->CreateGroupsWidget();
+
+	if (NestedTree->HasAnyParameters())
+	{
+		this->ChildSlot
+		[
+			SNew(SVerticalBox)
+			+ SVerticalBox::Slot()
+			.Padding(0.0f)
+			[
+				
+				SNew(SVerticalBox)
+				+ SVerticalBox::Slot()
+				.AutoHeight()
+				[
+					SAssignNew(HeaderBox, SHorizontalBox)
+				]
+				+ SVerticalBox::Slot()
+				.Padding(0.0f)
+				[
+					SNew(SHorizontalBox)
+					+ SHorizontalBox::Slot()
+					.HAlign(HAlign_Fill)
+					[
+						NestedTree.ToSharedRef()
+					]
+					+ SHorizontalBox::Slot()
+					.HAlign(HAlign_Right)
+					.AutoWidth()
+					[
+						SNew(SBox)
+						.WidthOverride(16.0f)
+						[
+							ExternalScrollbar.ToSharedRef()
+						]
+					]
+					
+				]
+				
+			]
+		];
+
+		HeaderBox->AddSlot()
+			.FillWidth(1.0f)
+			[
+				SNullWidget::NullWidget
+			];
+	}
+	else
+	{
+		this->ChildSlot
+			[
+				SNew(SBox)
+				.Padding(10.0f)
+				.HAlign(HAlign_Center)
+				[
+					SNew(STextBlock)
+					.Text(LOCTEXT("ConnectMaterialParametersToFillList", "Connect a parameter to see it here."))
+				]
+			];
+	}
+}
+
+void SMaterialDynamicParametersPanelWidget::Construct(const FArguments& InArgs)
+{
+	MaterialInstance = InArgs._InMaterialInstance;
+	ExternalScrollbar = SNew(SScrollBar);
+
+	FPropertyEditorModule& PropertyEditorModule = FModuleManager::LoadModuleChecked<FPropertyEditorModule>("PropertyEditor");
+	const FPropertyRowGeneratorArgs PropertyRowGeneratorArgs;
+	Generator = PropertyEditorModule.CreatePropertyRowGenerator(PropertyRowGeneratorArgs);
+
+	Generator->SetObjects({ MaterialInstance.Get() });
+
+	NestedTree = SNew(SMaterialDynamicParametersOverviewTree)
+		.InMaterialInstance(InArgs._InMaterialInstance)
+		.InOwner(SharedThis(this))
+		.InScrollbar(ExternalScrollbar);
+
+	
+	Refresh();
+}
+
+void SMaterialDynamicParametersPanelWidget::UpdateInstance(UMaterialInstance* InMaterialInstance)
+{
+	Generator->SetObjects({ InMaterialInstance });
+	NestedTree->MaterialInstance = InMaterialInstance;
+	Refresh();
+}
+
+#undef LOCTEXT_NAMESPACE
diff --git a/Engine/Plugins/Experimental/InlineMaterialInstance/Source/InlineMaterialInstanceEditor/Private/Widgets/SMaterialDynamicParametersPanelWidget.h b/Engine/Plugins/Experimental/InlineMaterialInstance/Source/InlineMaterialInstanceEditor/Private/Widgets/SMaterialDynamicParametersPanelWidget.h
new file mode 100644
--- /dev/null
+++ b/Engine/Plugins/Experimental/InlineMaterialInstance/Source/InlineMaterialInstanceEditor/Private/Widgets/SMaterialDynamicParametersPanelWidget.h
@@ -0,0 +1,51 @@
+// Copyright Epic Games, Inc. All Rights Reserved.
+
+#pragma once
+
+#include "CoreMinimal.h"
+#include "Widgets/DeclarativeSyntaxSupport.h"
+#include "Widgets/SCompoundWidget.h"
+
+class IPropertyRowGenerator;
+class SMaterialDynamicParametersOverviewTree;
+class UMaterialInstance;
+
+/**
+ * Dynamic Material parameters view. The view creates all widgets for dynamic material instance parameters
+ * Most of the logic copied from SMaterialParametersOverviewTreeItem
+ */
+class SMaterialDynamicParametersPanelWidget : public SCompoundWidget
+{
+	friend SMaterialDynamicParametersOverviewTree;
+	
+public:
+	SLATE_BEGIN_ARGS(SMaterialDynamicParametersPanelWidget)
+		: _InMaterialInstance(nullptr)
+	{}
+
+	SLATE_ARGUMENT(UMaterialInstance*, InMaterialInstance)
+	SLATE_END_ARGS()
+
+	void Construct(const FArguments& InArgs);
+
+	/** Update material instance for tree generation and refresh */
+	void UpdateInstance(UMaterialInstance* InMaterialInstance);
+
+private:
+	/** Refresh MaterialParameters tree */
+	void Refresh();
+
+private:
+	/** The set of material parameters this is associated with */
+	TWeakObjectPtr<UMaterialInstance> MaterialInstance;
+
+	/** The tree contained in this item */
+	TSharedPtr<class SMaterialDynamicParametersOverviewTree> NestedTree;
+
+	/** Optional external Scroll bar */
+	TSharedPtr<class SScrollBar> ExternalScrollbar;
+
+	/** Row generator instance */
+	TSharedPtr<IPropertyRowGenerator> Generator;
+};
+
diff --git a/Engine/Plugins/Experimental/InlineMaterialInstance/Source/InlineMaterialInstanceEditor/Private/Widgets/SMaterialDynamicWidgets.cpp b/Engine/Plugins/Experimental/InlineMaterialInstance/Source/InlineMaterialInstanceEditor/Private/Widgets/SMaterialDynamicWidgets.cpp
new file mode 100644
--- /dev/null
+++ b/Engine/Plugins/Experimental/InlineMaterialInstance/Source/InlineMaterialInstanceEditor/Private/Widgets/SMaterialDynamicWidgets.cpp
@@ -0,0 +1,254 @@
+﻿// Copyright Epic Games, Inc. All Rights Reserved.
+
+#include "SMaterialDynamicWidgets.h"
+
+#include "Components/ActorComponent.h"
+#include "Editor.h"
+#include "MaterialEditingLibrary.h"
+#include "Materials/MaterialInstanceDynamic.h"
+#include "PropertyCustomizationHelpers.h"
+#include "SMaterialDynamicParametersPanelWidget.h"
+#include "ScopedTransaction.h"
+#include "Subsystems/AssetEditorSubsystem.h"
+#include "Widgets/Input/SButton.h"
+#include "Widgets/Layout/SBox.h"
+#include "Widgets/Text/STextBlock.h"
+
+#define LOCTEXT_NAMESPACE "FMaterialDynamicList"
+
+void SMaterialDynamicView::Construct(const FArguments& InArgs, const TSharedRef<FMaterialItemView>& InMaterialItemView, UActorComponent* InCurrentComponent)
+{
+	MaterialItemViewWeakPtr = InMaterialItemView;
+	CurrentComponent = InCurrentComponent;
+
+	TSharedRef<SVerticalBox> ResultWidget = SNew(SVerticalBox);
+	if (InCurrentComponent != nullptr)
+	{
+		UMaterialInterface* MaterialInterface = InMaterialItemView->GetMaterialListItem().Material.Get();
+		UMaterialInstance* MaterialInstance = Cast<UMaterialInstance>(MaterialInterface);
+		UMaterialInstanceDynamic* MaterialInstanceDynamic = Cast<UMaterialInstanceDynamic>(MaterialInterface);
+		
+		if (MaterialInstance)
+		{
+			const bool bIsInvertedVisibility = !!MaterialInstanceDynamic;
+			const TAttribute<EVisibility> CreateDynamicMaterialButtonVisibility = GetButtonVisibilityAttribute<UMaterialInstance>(bIsInvertedVisibility);
+			
+			ResultWidget->AddSlot()
+			.AutoHeight()
+			[
+				SNew(SBox)
+				[
+					SNew(SHorizontalBox)
+					+SHorizontalBox::Slot()
+					.AutoWidth()
+					[
+						SNew(SButton)
+						.Visibility(GetButtonVisibilityAttribute<UMaterialInstanceDynamic>())
+						.ToolTipText(LOCTEXT("Revert_Button_Tooltip", "Revert the Dynamic Material Instance back to the original Material Instance"))
+						.OnClicked(this, &SMaterialDynamicView::OnRevertButtonClicked)
+						[
+							SNew(STextBlock).Text(LOCTEXT("Revert_Button", "Revert"))
+						]
+					]
+
+					+SHorizontalBox::Slot()
+					.AutoWidth()
+					[
+						SNew(SButton)
+						.Visibility(GetButtonVisibilityAttribute<UMaterialInstanceDynamic>())
+						.ToolTipText(LOCTEXT("Reset_Button_Tooltip", "Reset the properties to the Material Instance default"))
+						.OnClicked(this, &SMaterialDynamicView::OnResetButtonClicked)
+						[
+							SNew(STextBlock).Text(LOCTEXT("Reset_Button", "Reset"))
+						]
+					]
+
+					+SHorizontalBox::Slot()
+					.AutoWidth()
+					[
+						SNew(SButton)
+						.Visibility(GetButtonVisibilityAttribute<UMaterialInstanceDynamic>())
+						.ToolTipText(LOCTEXT("CopyToOriginal_Button_Tooltip", "Copy and overwrite the parameters onto the original Material Instance"))
+						.OnClicked(this, &SMaterialDynamicView::OnCopyToOriginalButtonClicked)
+						[
+							SNew(STextBlock).Text(LOCTEXT("CopyToOriginal_Button", "Copy to Original"))
+						]
+					]
+
+					+SHorizontalBox::Slot()
+					.AutoWidth()
+					[
+						SNew(SButton)
+						.Visibility(CreateDynamicMaterialButtonVisibility)
+						.ToolTipText(LOCTEXT("CreateDynamicMaterial_Button_Tooltip", "Create a Dynamic Material Instance for this Material Instance and automatically sets it"))
+						.OnClicked(this, &SMaterialDynamicView::OnCreateDynamicMaterialButtonClicked)
+						[
+							SNew(STextBlock).Text(LOCTEXT("CreateDynamicMaterial_Button", "Create Dynamic Material"))
+						]
+					]
+				]
+			];
+		}
+
+		if (MaterialInstanceDynamic)
+		{
+			ResultWidget->AddSlot()
+			.AutoHeight()
+			[
+				SNew(SMaterialDynamicParametersPanelWidget)
+				.InMaterialInstance(MaterialInstanceDynamic)
+			];
+		}
+	}
+
+	ChildSlot
+	[
+		ResultWidget
+	];
+}
+
+FReply SMaterialDynamicView::OnResetButtonClicked() const
+{
+	UActorComponent* OwnerActorComponent = CurrentComponent.Get();
+	if (!ensure(OwnerActorComponent))
+	{
+		return FReply::Handled();
+	}
+
+	const TSharedPtr<FMaterialItemView> MaterialItemView = MaterialItemViewWeakPtr.Pin();
+	if(!ensure(MaterialItemView.IsValid()))
+	{
+		return FReply::Handled();
+	}
+	
+	UMaterialInstanceDynamic* MaterialInstanceDynamic = Cast<UMaterialInstanceDynamic>(MaterialItemView->GetMaterialListItem().Material.Get());
+	if (!MaterialInstanceDynamic)
+	{
+		return FReply::Handled();
+	}
+
+	UMaterialInstance* ParentMaterialInstance = Cast<UMaterialInstance>(MaterialInstanceDynamic->Parent);
+	if (!ParentMaterialInstance)
+	{
+		return FReply::Handled();
+	}
+
+
+	FScopedTransaction Transaction(LOCTEXT("ResetDynamicMaterial", "Reset Dynamic Material"));
+	MaterialInstanceDynamic->Modify();
+	OwnerActorComponent->Modify();
+	MaterialInstanceDynamic->CopyParameterOverrides(ParentMaterialInstance);
+
+	return FReply::Handled();
+}
+
+FReply SMaterialDynamicView::OnRevertButtonClicked() const
+{
+	UActorComponent* OwnerActorComponent = CurrentComponent.Get();
+	if (!ensure(OwnerActorComponent))
+	{
+		return FReply::Handled();
+	}
+
+	const TSharedPtr<FMaterialItemView> MaterialItemView = MaterialItemViewWeakPtr.Pin();
+	if(!ensure(MaterialItemView.IsValid()))
+	{
+		return FReply::Handled();
+	}
+	
+	UMaterialInstance* ParentMaterialInstance = GetMaterialParent<UMaterialInstance, UMaterialInstanceDynamic>(MaterialItemView->GetMaterialListItem().Material.Get());
+	if (!ParentMaterialInstance)
+	{
+		return FReply::Handled();
+	}
+
+	FScopedTransaction Transaction(LOCTEXT("RevertDynamicMaterial", "Revert Dynamic Material"));
+	MaterialItemView->ReplaceMaterial(ParentMaterialInstance);
+	
+	return FReply::Handled();
+}
+
+FReply SMaterialDynamicView::OnCopyToOriginalButtonClicked() const
+{
+	UActorComponent* OwnerActorComponent = CurrentComponent.Get();
+	if (!ensure(OwnerActorComponent))
+	{
+		return FReply::Handled();
+	}
+
+	const TSharedPtr<FMaterialItemView> MaterialItemView = MaterialItemViewWeakPtr.Pin();
+	if(!ensure(MaterialItemView.IsValid()))
+	{
+		return FReply::Handled();
+	}
+	
+	UMaterialInstanceDynamic* MaterialInstanceDynamic = Cast<UMaterialInstanceDynamic>(MaterialItemView->GetMaterialListItem().Material.Get());
+	if (!MaterialInstanceDynamic)
+	{
+		return FReply::Handled();
+	}
+
+	UMaterialInstance* ParentMaterialInstance = Cast<UMaterialInstance>(MaterialInstanceDynamic->Parent);
+	if (!ParentMaterialInstance)
+	{
+		return FReply::Handled();
+	}
+
+	FScopedTransaction Transaction(LOCTEXT("CopyToOriginalMaterial", "Copy To Original Material"));
+	ParentMaterialInstance->ScalarParameterValues = MaterialInstanceDynamic->ScalarParameterValues;
+	ParentMaterialInstance->VectorParameterValues = MaterialInstanceDynamic->VectorParameterValues;
+	ParentMaterialInstance->TextureParameterValues = MaterialInstanceDynamic->TextureParameterValues;
+	ParentMaterialInstance->RuntimeVirtualTextureParameterValues = MaterialInstanceDynamic->RuntimeVirtualTextureParameterValues;
+	ParentMaterialInstance->FontParameterValues = MaterialInstanceDynamic->FontParameterValues;
+	ParentMaterialInstance->UpdateStaticPermutation();
+	ParentMaterialInstance->Modify();
+
+#if WITH_EDITOR
+	// Check if the Parent is being Edited actively in a Window somewhere
+	if (GEditor->GetEditorSubsystem<UAssetEditorSubsystem>()->FindEditorForAsset(ParentMaterialInstance, false))
+	{
+		UMaterialEditingLibrary::RecompileMaterial(ParentMaterialInstance->GetMaterial());
+	}
+#endif
+
+	return FReply::Handled();
+}
+
+FReply SMaterialDynamicView::OnCreateDynamicMaterialButtonClicked() const
+{
+	const TSharedPtr<FMaterialItemView> MaterialItemView = MaterialItemViewWeakPtr.Pin();
+	if(!ensure(MaterialItemView.IsValid()))
+	{
+		return FReply::Handled();
+	}
+	
+	UMaterialInstance* MaterialInstance = Cast<UMaterialInstance>(MaterialItemView->GetMaterialListItem().Material.Get());
+	UActorComponent* OwnerActorComponent = CurrentComponent.Get();
+
+	if (!OwnerActorComponent)
+	{
+		return FReply::Handled();
+	}
+
+	if (MaterialInstance)
+	{
+		OwnerActorComponent->Modify();
+
+		FScopedTransaction Transaction(LOCTEXT("CreateDynamicMaterial", "Create Dynamic Material"));
+
+		// Create Transactional Material. That is allawed create MID with MU and transact the changes
+		UMaterialInstanceDynamic* NewMaterialInstanceDynamic = NewObject<UMaterialInstanceDynamic>(OwnerActorComponent, NAME_None, RF_Public | RF_Transactional);
+		NewMaterialInstanceDynamic->Parent = MaterialInstance;
+		NewMaterialInstanceDynamic->UpdateCachedData();
+
+
+		// Copy Parameter Overrides with Transaction
+		NewMaterialInstanceDynamic->CopyParameterOverrides(MaterialInstance);
+		MaterialItemView->ReplaceMaterial(NewMaterialInstanceDynamic);
+	}
+
+	return FReply::Handled();
+}
+
+#undef LOCTEXT_NAMESPACE
+
diff --git a/Engine/Plugins/Experimental/InlineMaterialInstance/Source/InlineMaterialInstanceEditor/Private/Widgets/SMaterialDynamicWidgets.h b/Engine/Plugins/Experimental/InlineMaterialInstance/Source/InlineMaterialInstanceEditor/Private/Widgets/SMaterialDynamicWidgets.h
new file mode 100644
--- /dev/null
+++ b/Engine/Plugins/Experimental/InlineMaterialInstance/Source/InlineMaterialInstanceEditor/Private/Widgets/SMaterialDynamicWidgets.h
@@ -0,0 +1,78 @@
+﻿// Copyright Epic Games, Inc. All Rights Reserved.
+
+#pragma once
+
+#include "CoreMinimal.h"
+#include "MaterialList.h"
+#include "Widgets/DeclarativeSyntaxSupport.h"
+#include "Widgets/SCompoundWidget.h"
+
+class IPropertyHandle;
+class IDetailLayoutBuilder;
+class UActorComponent;
+class UMaterialInterface;
+
+/**
+ * Creates the view for material dynamic view
+ */
+class SMaterialDynamicView : public SCompoundWidget
+{
+public:
+	SLATE_BEGIN_ARGS(SMaterialDynamicView)
+	{}
+	SLATE_END_ARGS()
+
+	void Construct(const FArguments& InArgs, const TSharedRef<FMaterialItemView>& InMaterialItemView, UActorComponent* InCurrentComponent);
+
+private:
+	/** Get Visibility for extension button */
+	template<typename TMaterialClass>
+	TAttribute<EVisibility> GetButtonVisibilityAttribute(const bool bIsInvertedResult = false) const
+	{
+		return MakeAttributeLambda([this, bIsInvertedResult]()
+			{
+				const TSharedPtr<FMaterialItemView> MaterialItemView = MaterialItemViewWeakPtr.Pin();
+				if(!ensure(MaterialItemView.IsValid()))
+				{
+					return EVisibility::Collapsed;
+				}
+			
+				bool IsVisible = Cast<TMaterialClass>(MaterialItemView->GetMaterialListItem().Material.Get()) ? true : false;
+				IsVisible ^= bIsInvertedResult;
+
+				return IsVisible ? EVisibility::Visible : EVisibility::Collapsed;
+			});
+	};
+
+	/** Button handler for revert to original material for the slot */
+	FReply OnRevertButtonClicked() const;
+
+	/** Button handler for reset parameters to default instance material parameter */
+	FReply OnResetButtonClicked() const;
+
+	/** Button handler for copy material parameters to original instance material */
+	FReply OnCopyToOriginalButtonClicked() const;
+
+	/** Button handler for create dynamic material */
+	FReply OnCreateDynamicMaterialButtonClicked() const;
+
+	/** Get parent material and cast to specified type */
+	template<typename TReturnType, typename TInputType>
+	TReturnType* GetMaterialParent(UMaterialInterface* InMaterialInterface) const
+	{
+		if (TInputType* Material = Cast<TInputType>(InMaterialInterface))
+		{
+			return Cast<TReturnType>(Material->Parent);
+		}
+
+		return nullptr;
+	}
+
+private:
+	/** Weak Pointer to Material Item View */
+	TWeakPtr<FMaterialItemView> MaterialItemViewWeakPtr;
+
+	/** Weak pointer to current component using the current material */
+	TWeakObjectPtr<UActorComponent> CurrentComponent;
+};
+
diff --git a/Engine/Plugins/Experimental/VirtualProductionUtilities/Source/VPMaterialsEditor/Private/SMaterialDynamicParametersPanelWidget.cpp b/Engine/Plugins/Experimental/VirtualProductionUtilities/Source/VPMaterialsEditor/Private/SMaterialDynamicParametersPanelWidget.cpp
deleted file mode 100644
diff --git a/Engine/Plugins/Experimental/VirtualProductionUtilities/Source/VPMaterialsEditor/Private/SMaterialDynamicParametersPanelWidget.h b/Engine/Plugins/Experimental/VirtualProductionUtilities/Source/VPMaterialsEditor/Private/SMaterialDynamicParametersPanelWidget.h
deleted file mode 100644
diff --git a/Engine/Plugins/Experimental/VirtualProductionUtilities/Source/VPMaterialsEditor/Private/SMaterialDynamicWidgets.cpp b/Engine/Plugins/Experimental/VirtualProductionUtilities/Source/VPMaterialsEditor/Private/SMaterialDynamicWidgets.cpp
deleted file mode 100644
diff --git a/Engine/Plugins/Experimental/VirtualProductionUtilities/Source/VPMaterialsEditor/Private/SMaterialDynamicWidgets.h b/Engine/Plugins/Experimental/VirtualProductionUtilities/Source/VPMaterialsEditor/Private/SMaterialDynamicWidgets.h
deleted file mode 100644
diff --git a/Engine/Plugins/Experimental/VirtualProductionUtilities/Source/VPMaterialsEditor/Private/VPMaterialsEditorModule.cpp b/Engine/Plugins/Experimental/VirtualProductionUtilities/Source/VPMaterialsEditor/Private/VPMaterialsEditorModule.cpp
deleted file mode 100644
diff --git a/Engine/Plugins/Experimental/VirtualProductionUtilities/Source/VPMaterialsEditor/Private/VPMaterialsEditorModule.h b/Engine/Plugins/Experimental/VirtualProductionUtilities/Source/VPMaterialsEditor/Private/VPMaterialsEditorModule.h
deleted file mode 100644
diff --git a/Engine/Plugins/Experimental/VirtualProductionUtilities/Source/VPMaterialsEditor/VPMaterialsEditor.Build.cs b/Engine/Plugins/Experimental/VirtualProductionUtilities/Source/VPMaterialsEditor/VPMaterialsEditor.Build.cs
deleted file mode 100644
diff --git a/Engine/Plugins/Experimental/VirtualProductionUtilities/VirtualProductionUtilities.uplugin b/Engine/Plugins/Experimental/VirtualProductionUtilities/VirtualProductionUtilities.uplugin
--- a/Engine/Plugins/Experimental/VirtualProductionUtilities/VirtualProductionUtilities.uplugin
+++ b/Engine/Plugins/Experimental/VirtualProductionUtilities/VirtualProductionUtilities.uplugin
@@ -33,11 +33,6 @@
 			"Name": "VPBookmarkEditor",
 			"Type": "Editor",
 			"LoadingPhase": "PostDefault"
-		},
-		{
-			"Name": "VPMaterialsEditor",
-			"Type": "Editor",
-			"LoadingPhase": "PostDefault"
 		}
 	],
 	"Plugins": [
@@ -84,6 +79,10 @@
 		{
 			"Name": "CineCameraSceneCapture",
 			"Enabled": true
+		},
+		{
+			"Name": "InlineMaterialInstance",
+			"Enabled": true
 		}
 	]
 }
