Added widget validation to the Inline Material Instance plugin

diff --git a/Engine/Plugins/Experimental/InlineMaterialInstance/Source/InlineMaterialInstanceEditor/Private/Widgets/SMaterialDynamicParametersPanelWidget.cpp b/Engine/Plugins/Experimental/InlineMaterialInstance/Source/InlineMaterialInstanceEditor/Private/Widgets/SMaterialDynamicParametersPanelWidget.cpp
--- a/Engine/Plugins/Experimental/InlineMaterialInstance/Source/InlineMaterialInstanceEditor/Private/Widgets/SMaterialDynamicParametersPanelWidget.cpp
+++ b/Engine/Plugins/Experimental/InlineMaterialInstance/Source/InlineMaterialInstanceEditor/Private/Widgets/SMaterialDynamicParametersPanelWidget.cpp
@@ -317,19 +317,16 @@ void SMaterialDynamicParametersOverviewTreeItem::Construct(const FArguments& InA
 	auto GenerateLeftAndRightWidgets = [&](IDetailTreeNode& InNode)
 	{
 		const FNodeWidgets NodeWidgets = InNode.CreateNodeWidgets();
-		LeftSideWidget = NodeWidgets.NameWidget.ToSharedRef();
-
-		RightSideWidget = SNew(SHorizontalBox)
-				+SHorizontalBox::Slot()
-				.FillWidth(1.f)
-				[
-					NodeWidgets.ValueWidget.ToSharedRef()
-				]
-				+ SHorizontalBox::Slot()
-				.AutoWidth()
-				[
-					GetRowExtensionButtons(InNode.CreatePropertyHandle())
-				];
+		if (NodeWidgets.NameWidget.IsValid() && NodeWidgets.ValueWidget.IsValid())
+		{
+			LeftSideWidget = NodeWidgets.NameWidget.ToSharedRef();
+			RightSideWidget = SNew(SHorizontalBox)
+					+SHorizontalBox::Slot()
+					.FillWidth(1.f)
+					[
+						NodeWidgets.ValueWidget.ToSharedRef()
+					];
+		}
 	};
 
 // GROUP --------------------------------------------------
