Virtual Production Inline Material Editor: Added max height to the parameter list to avoid SListView issues.

diff --git a/Engine/Plugins/Experimental/InlineMaterialInstance/Source/InlineMaterialInstanceEditor/Private/Widgets/SMaterialDynamicParametersPanelWidget.cpp b/Engine/Plugins/Experimental/InlineMaterialInstance/Source/InlineMaterialInstanceEditor/Private/Widgets/SMaterialDynamicParametersPanelWidget.cpp
--- a/Engine/Plugins/Experimental/InlineMaterialInstance/Source/InlineMaterialInstanceEditor/Private/Widgets/SMaterialDynamicParametersPanelWidget.cpp
+++ b/Engine/Plugins/Experimental/InlineMaterialInstance/Source/InlineMaterialInstanceEditor/Private/Widgets/SMaterialDynamicParametersPanelWidget.cpp
@@ -686,6 +686,7 @@ void SMaterialDynamicParametersPanelWidget::Refresh()
 				]
 				+ SVerticalBox::Slot()
 				.Padding(0.0f)
+				.MaxHeight(300.f)
 				[
 					SNew(SHorizontalBox)
 					+ SHorizontalBox::Slot()
